Fix for Opus assert with 44.1 looping sounds
#tests inline and streaming 44.1 in editor in opus

diff --git a/Engine/Source/Runtime/OpusAudioDecoder/Module/Private/OpusAudioInfo.cpp b/Engine/Source/Runtime/OpusAudioDecoder/Module/Private/OpusAudioInfo.cpp
--- a/Engine/Source/Runtime/OpusAudioDecoder/Module/Private/OpusAudioInfo.cpp
+++ b/Engine/Source/Runtime/OpusAudioDecoder/Module/Private/OpusAudioInfo.cpp
@@ -211,14 +211,33 @@ FDecodeResult FOpusAudioInfo::Decode(const uint8* CompressedData, const int32 Co
 			// Get anything that is left over from the previous call.
 			if (PreviousDecodedUnusedSamples.Num())
 			{
-				check(NumRemainingSamplesToSkip == 0);
-				int32 MaxToCopyOut = OutputSizeToGo >= PreviousDecodedUnusedSamples.Num() ? PreviousDecodedUnusedSamples.Num() : OutputSizeToGo;
-				FMemory::Memcpy(OutputDataPtr, PreviousDecodedUnusedSamples.GetData(), MaxToCopyOut);
-				PreviousDecodedUnusedSamples.RemoveAt(0, MaxToCopyOut);
-				OutputSizeToGo -= MaxToCopyOut;
-				OutputDataPtr += MaxToCopyOut;
-				// If there is still something left over then we are done here.
-				Result.NumAudioFramesProduced += MaxToCopyOut / kFrameBytes;
+				// Skip some of these residual samples.
+				if (NumRemainingSamplesToSkip)
+				{
+					if (NumRemainingSamplesToSkip >= PreviousDecodedUnusedSamples.Num())
+					{
+						NumRemainingSamplesToSkip -= PreviousDecodedUnusedSamples.Num();
+						PreviousDecodedUnusedSamples.Empty();					
+					}
+					else
+					{
+						PreviousDecodedUnusedSamples.RemoveAt(0, NumRemainingSamplesToSkip);
+						NumRemainingSamplesToSkip = 0;
+					}
+				}
+
+				// Recheck if there's anything in previous (we may have just removed some).
+				if (PreviousDecodedUnusedSamples.Num())
+				{
+					int32 MaxToCopyOut = OutputSizeToGo >= PreviousDecodedUnusedSamples.Num() ? PreviousDecodedUnusedSamples.Num() : OutputSizeToGo;
+					FMemory::Memcpy(OutputDataPtr, PreviousDecodedUnusedSamples.GetData(), MaxToCopyOut);
+					PreviousDecodedUnusedSamples.RemoveAt(0, MaxToCopyOut);
+					OutputSizeToGo -= MaxToCopyOut;
+					OutputDataPtr += MaxToCopyOut;
+					// If there is still something left over then we are done here.
+					Result.NumAudioFramesProduced += MaxToCopyOut / kFrameBytes;
+				}
+
 				if (OutputSizeToGo == 0 || PreviousDecodedUnusedSamples.Num())
 				{
 					Result.NumCompressedBytesConsumed = InputDataPtr - CompressedData;
@@ -387,12 +406,7 @@ public:
 
 	virtual void StartupModule() override
 	{
-		constexpr FSimpleAudioInfoFactory::FCapabilities Caps = 
-		{
-			/*bSupportsSeeking*/ true,				
-			/*bSupportsSeekableStreaming*/ false,	
-		};
-		Factory = MakeUnique<FSimpleAudioInfoFactory>([] { return new FOpusAudioInfo(); }, Audio::NAME_OPUS, Caps);
+		Factory = MakeUnique<FSimpleAudioInfoFactory>([] { return new FOpusAudioInfo(); }, Audio::NAME_OPUS);
 	}
 
 	virtual void ShutdownModule() override {}
