Fixed incorrect skipping over initial samples in Opus decoder. Previous fix was incorrect in that it did not skip over actual number of bytes but number of samples only.

#rnx

diff --git a/Engine/Source/Runtime/OpusAudioDecoder/Module/Private/OpusAudioInfo.cpp b/Engine/Source/Runtime/OpusAudioDecoder/Module/Private/OpusAudioInfo.cpp
--- a/Engine/Source/Runtime/OpusAudioDecoder/Module/Private/OpusAudioInfo.cpp
+++ b/Engine/Source/Runtime/OpusAudioDecoder/Module/Private/OpusAudioInfo.cpp
@@ -202,7 +202,7 @@ FDecodeResult FOpusAudioInfo::Decode(const uint8* CompressedData, const int32 Co
 		const int32 kFrameBytes = NumChannels * sizeof(int16);
 		int32 OutputSizeToGo = OutputPCMDataSize;
 		int32 CompressedSizeToGo = CompressedDataSize;
-		const uint8* InputDataPtr = CompressedData; 
+		const uint8* InputDataPtr = CompressedData;
 		uint8* OutputDataPtr = OutPCMData;
 
 		Result.NumAudioFramesProduced = 0;
@@ -211,33 +211,14 @@ FDecodeResult FOpusAudioInfo::Decode(const uint8* CompressedData, const int32 Co
 			// Get anything that is left over from the previous call.
 			if (PreviousDecodedUnusedSamples.Num())
 			{
-				// Skip some of these residual samples.
-				if (NumRemainingSamplesToSkip)
-				{
-					if (NumRemainingSamplesToSkip >= PreviousDecodedUnusedSamples.Num())
-					{
-						NumRemainingSamplesToSkip -= PreviousDecodedUnusedSamples.Num();
-						PreviousDecodedUnusedSamples.Empty();					
-					}
-					else
-					{
-						PreviousDecodedUnusedSamples.RemoveAt(0, NumRemainingSamplesToSkip);
-						NumRemainingSamplesToSkip = 0;
-					}
-				}
-
-				// Recheck if there's anything in previous (we may have just removed some).
-				if (PreviousDecodedUnusedSamples.Num())
-				{
-					int32 MaxToCopyOut = OutputSizeToGo >= PreviousDecodedUnusedSamples.Num() ? PreviousDecodedUnusedSamples.Num() : OutputSizeToGo;
-					FMemory::Memcpy(OutputDataPtr, PreviousDecodedUnusedSamples.GetData(), MaxToCopyOut);
-					PreviousDecodedUnusedSamples.RemoveAt(0, MaxToCopyOut);
-					OutputSizeToGo -= MaxToCopyOut;
-					OutputDataPtr += MaxToCopyOut;
-					// If there is still something left over then we are done here.
-					Result.NumAudioFramesProduced += MaxToCopyOut / kFrameBytes;
-				}
-
+				check(NumRemainingSamplesToSkip == 0);
+				int32 MaxToCopyOut = OutputSizeToGo >= PreviousDecodedUnusedSamples.Num() ? PreviousDecodedUnusedSamples.Num() : OutputSizeToGo;
+				FMemory::Memcpy(OutputDataPtr, PreviousDecodedUnusedSamples.GetData(), MaxToCopyOut);
+				PreviousDecodedUnusedSamples.RemoveAt(0, MaxToCopyOut);
+				OutputSizeToGo -= MaxToCopyOut;
+				OutputDataPtr += MaxToCopyOut;
+				// If there is still something left over then we are done here.
+				Result.NumAudioFramesProduced += MaxToCopyOut / kFrameBytes;
 				if (OutputSizeToGo == 0 || PreviousDecodedUnusedSamples.Num())
 				{
 					Result.NumCompressedBytesConsumed = InputDataPtr - CompressedData;
@@ -313,6 +294,20 @@ FDecodeResult FOpusAudioInfo::Decode(const uint8* CompressedData, const int32 Co
 					PreviousDecodedUnusedSamples.SetNum(NumPrevNow);
 					CompressedSizeToGo -= ActualChunkSize;
 					InputDataPtr += ActualChunkSize;
+					if (NumRemainingSamplesToSkip)
+					{
+						const int32 NumBytesToSkip = NumRemainingSamplesToSkip * kFrameBytes;
+						if (NumBytesToSkip >= NumPrevNow)
+						{
+							PreviousDecodedUnusedSamples.SetNum(0, EAllowShrinking::No);
+							NumRemainingSamplesToSkip -= NumPrevNow / kFrameBytes;
+						}
+						else
+						{
+							PreviousDecodedUnusedSamples.RemoveAt(0, NumBytesToSkip);
+							NumRemainingSamplesToSkip = 0;
+						}
+					}
 					continue;
 				}
 				else
@@ -401,7 +396,7 @@ void FOpusAudioInfo::SeekToFrame(const uint32 InSeekFrame)
 
 class OPUSAUDIODECODER_API FOpusAudioDecoderModule : public IModuleInterface
 {
-public:	
+public:
 	TUniquePtr<IAudioInfoFactory> Factory;
 
 	virtual void StartupModule() override
