Opus: Seek to use the chunk containing the target, skipping over the first unused samples instead of starting at the chunk following the target.

#rnx

diff --git a/Engine/Source/Runtime/OpusAudioDecoder/Module/Private/OpusAudioInfo.cpp b/Engine/Source/Runtime/OpusAudioDecoder/Module/Private/OpusAudioInfo.cpp
--- a/Engine/Source/Runtime/OpusAudioDecoder/Module/Private/OpusAudioInfo.cpp
+++ b/Engine/Source/Runtime/OpusAudioDecoder/Module/Private/OpusAudioInfo.cpp
@@ -356,6 +356,7 @@ void FOpusAudioInfo::SeekToTime(const float InSeekTime)
 
 void FOpusAudioInfo::SeekToFrame(const uint32 InSeekFrame)
 {
+	int32 NumSamplesToSkipInChunk = 0;
 	if (GetStreamingSoundWave().IsValid())
 	{
 		IStreamedCompressedInfo::SeekToFrame(InSeekFrame);
@@ -373,10 +374,11 @@ void FOpusAudioInfo::SeekToFrame(const uint32 InSeekFrame)
 			int32 ExpectedFrameSize = opus_packet_get_samples_per_frame(ChunkPtr + 2, Header.EncodedSampleRate);
 			int32 NumExpectedTotal = ExpectedFrames * ExpectedFrameSize * NumChannels;
 
-			if (CurrentChunkSampleNum >= SeekSampleNum && SeekSampleNum < CurrentChunkSampleNum + NumExpectedTotal)
+			if (SeekSampleNum >= CurrentChunkSampleNum && SeekSampleNum < CurrentChunkSampleNum + NumExpectedTotal)
 			{
 				CurrentSampleCount = CurrentChunkSampleNum;
 				SrcBufferOffset = ChunkPtr - SrcBufferData;
+				NumSamplesToSkipInChunk = SeekSampleNum - CurrentChunkSampleNum;
 				break;
 			}
 			CurrentChunkSampleNum += NumExpectedTotal;
@@ -390,7 +392,7 @@ void FOpusAudioInfo::SeekToFrame(const uint32 InSeekFrame)
 		}
 	}
 
-	NumRemainingSamplesToSkip = Header.NumSilentSamplesAtBeginning + Header.NumPreSkipSamples;
+	NumRemainingSamplesToSkip = Header.NumSilentSamplesAtBeginning + Header.NumPreSkipSamples + NumSamplesToSkipInChunk;
 	PreviousDecodedUnusedSamples.Empty();
 }
 
